Clarify sample rate limitations for RAD Audio.

diff --git a/Engine/Source/Developer/AudioFormatRad/Private/AudioFormatRad.cpp b/Engine/Source/Developer/AudioFormatRad/Private/AudioFormatRad.cpp
--- a/Engine/Source/Developer/AudioFormatRad/Private/AudioFormatRad.cpp
+++ b/Engine/Source/Developer/AudioFormatRad/Private/AudioFormatRad.cpp
@@ -98,7 +98,7 @@ public:
 			UE_LOG(LogAudioFormatRad, Warning, TEXT("Failed to encode RAD Audio: %hs"),  RadAErrorString(RadCompressError));
 			if (RadCompressError == RADA_COMPRESS_ERROR_RATE)
 			{
-				UE_LOG(LogAudioFormatRad, Warning, TEXT("Sample rate provided: %u"), InQualityInfo.SampleRate);
+				UE_LOG(LogAudioFormatRad, Warning, TEXT("Sample rate provided: %u - please reimport at a valid sample rate."), InQualityInfo.SampleRate);
 			}
 			else if (RadCompressError == RADA_COMPRESS_ERROR_CHANS)
 			{
diff --git a/Engine/Source/Runtime/Engine/Classes/Sound/SoundWave.h b/Engine/Source/Runtime/Engine/Classes/Sound/SoundWave.h
--- a/Engine/Source/Runtime/Engine/Classes/Sound/SoundWave.h
+++ b/Engine/Source/Runtime/Engine/Classes/Sound/SoundWave.h
@@ -338,7 +338,7 @@ enum class ESoundAssetCompressionType : uint8
 	// The project defines the codec used for this asset.
 	ProjectDefined,
 
-	// As BinkAudio, except better quality. Comparable CPU usage.
+	// As BinkAudio, except better quality. Comparable CPU usage. Only valid sample rates are: 48000, 44100, 32000, and 24000.
 	RADAudio UMETA(DisplayName = "RAD Audio"),
 };
 
