Updating Rad Audio quality mapping to mostly be quality 5 since that's what it was tuned at, with some edges for fixing/adjusting per sound as needed.

diff --git a/Engine/Source/Developer/AudioFormatRad/Private/AudioFormatRad.cpp b/Engine/Source/Developer/AudioFormatRad/Private/AudioFormatRad.cpp
--- a/Engine/Source/Developer/AudioFormatRad/Private/AudioFormatRad.cpp
+++ b/Engine/Source/Developer/AudioFormatRad/Private/AudioFormatRad.cpp
@@ -17,19 +17,19 @@ namespace AudioFormatRadPrivate
 {
 	static uint8 GetCompressionLevelFromQualityIndex(const int32 InQualityIndex) 
 	{
-		// The engine default is 80, which we want to be 5. So we split a linear mapping
-		// if we are below or above that.
-		float MappedValue = 0;
-		if (InQualityIndex <= 80)
+		// RAD is tuned for 5, so we map almost everything to that, with some extremes 
+		// for manipulating edge cases. Project defaults vary from 40-80, but we make it symmetric:
+		uint8 RadLevel = 5;
+		if (InQualityIndex < 20) // 1..19 -> 1..4
 		{
-			MappedValue = FMath::GetMappedRangeValueClamped(FVector2d(1, 80), FVector2d(1, 5), InQualityIndex);
+			RadLevel = FMath::RoundToInt(FMath::GetMappedRangeValueClamped(FVector2d(1, 19), FVector2d(1, 4), InQualityIndex));
 		}
-		else
+		else if (InQualityIndex > 80) // 81..100 -> 6..9
 		{
-			MappedValue = FMath::GetMappedRangeValueClamped(FVector2d(80, 100), FVector2d(5, 9), InQualityIndex);
+			RadLevel = FMath::RoundToInt(FMath::GetMappedRangeValueClamped(FVector2d(81, 100), FVector2d(6, 9), InQualityIndex));
 		}
 
-		return (uint8)MappedValue;
+		return RadLevel;
 	}
 }
 
@@ -41,7 +41,7 @@ class FAudioFormatRad : public IAudioFormat
 	enum
 	{
 		/** Version for RAD Audio format, this becomes part of the DDC key. */
-		UE_AUDIO_RAD_VER = 2,
+		UE_AUDIO_RAD_VER = 3,
 	};
 
 public:
