Check header/lib match in rad audio at runtime.

#rnx

diff --git a/Engine/Source/Developer/AudioFormatRad/Private/AudioFormatRad.cpp b/Engine/Source/Developer/AudioFormatRad/Private/AudioFormatRad.cpp
--- a/Engine/Source/Developer/AudioFormatRad/Private/AudioFormatRad.cpp
+++ b/Engine/Source/Developer/AudioFormatRad/Private/AudioFormatRad.cpp
@@ -81,6 +81,15 @@ public:
 		TRACE_CPUPROFILER_EVENT_SCOPE(FAudioFormatRad::Cook);
 		check(InFormat == NAME_RADA);
 
+		#ifdef RADA_BUILD_VERSION
+		if (RadAGetBuildVersion() != RADA_BUILD_VERSION)
+		{
+			UE_LOG(LogAudioFormatRad, Error, TEXT("Header/library mismatch with Rad Audio! Header: %u, Library: %u - verify sync is correct"), RADA_BUILD_VERSION, RadAGetBuildVersion());
+			return false;
+		}
+		#endif
+
+
 		const uint8 CompressionLevel = AudioFormatRadPrivate::GetCompressionLevelFromQualityIndex(InQualityInfo.Quality);
 		
 		// If we're going to embed the seek-table in the stream, use -1 to give the largest table we can produce.
@@ -108,7 +117,7 @@ public:
 			}
 			else if (RadCompressError == RADA_COMPRESS_ERROR_CHANS)
 			{
-				UE_LOG(LogAudioFormatRad, Warning, TEXT("Channels provided: %u"), InQualityInfo.NumChannels);
+				UE_LOG(LogAudioFormatRad, Warning, TEXT("Channels provided: %u, max channels allowed: %d"), InQualityInfo.NumChannels, RADA_MAX_CHANS);
 			}
 			return false;
 		}
@@ -125,6 +134,14 @@ public:
 		TRACE_CPUPROFILER_EVENT_SCOPE(FAudioFormatRad::CookSurround);
 		check(InFormat == NAME_RADA);
 
+		#ifdef RADA_BUILD_VERSION
+		if (RadAGetBuildVersion() != RADA_BUILD_VERSION)
+		{
+			UE_LOG(LogAudioFormatRad, Error, TEXT("Header/library mismatch with Rad Audio! Header: %u, Library: %u - verify sync is correct"), RADA_BUILD_VERSION, RadAGetBuildVersion());
+			return false;
+		}
+		#endif
+
 		//
 		// CookSurround passes us a bunch of mono buffers, but RAD audio wants a standard 
 		// interleaved buffer
@@ -358,6 +375,14 @@ public:
 	}
 	virtual void StartupModule() override
 	{
+		#ifdef RADA_BUILD_VERSION
+		uint32_t LibraryBuildVersion = RadAGetBuildVersion();
+		if (LibraryBuildVersion != RADA_BUILD_VERSION)
+		{
+			UE_LOG(LogAudioFormatRad, Error, TEXT("Header/library mismatch with Rad Audio! Header: %u, Library: %u - verify sync is correct"), RADA_BUILD_VERSION, LibraryBuildVersion);
+		}
+		#endif
+
 		RadEncoder = new FAudioFormatRad();
 	}
 	virtual void ShutdownModule() override
